ARKit linker issue fixes

Not exporting interfaces

diff --git a/Engine/Plugins/Runtime/AR/ARUtilities/Source/ARUtilities/Public/ARLiveLinkRetargetAsset.h b/Engine/Plugins/Runtime/AR/ARUtilities/Source/ARUtilities/Public/ARLiveLinkRetargetAsset.h
--- a/Engine/Plugins/Runtime/AR/ARUtilities/Source/ARUtilities/Public/ARLiveLinkRetargetAsset.h
+++ b/Engine/Plugins/Runtime/AR/ARUtilities/Source/ARUtilities/Public/ARLiveLinkRetargetAsset.h
@@ -60,7 +60,7 @@ public:
 
 
 /** Interface that allows each platform to implement its own retarting logic */
-class ARUTILITIES_API IARLiveLinkRetargetingLogic : public IModularFeature
+class IARLiveLinkRetargetingLogic : public IModularFeature
 {
 public:
 	static const FName& GetModularFeatureName()
diff --git a/Engine/Plugins/Runtime/AR/AppleAR/AppleARKit/Source/AppleARKit/Public/AppleARKitFaceSupport.h b/Engine/Plugins/Runtime/AR/AppleAR/AppleARKit/Source/AppleARKit/Public/AppleARKitFaceSupport.h
--- a/Engine/Plugins/Runtime/AR/AppleAR/AppleARKit/Source/AppleARKit/Public/AppleARKitFaceSupport.h
+++ b/Engine/Plugins/Runtime/AR/AppleAR/AppleARKit/Source/AppleARKit/Public/AppleARKitFaceSupport.h
@@ -19,7 +19,7 @@ struct FGuid;
 class UTimecodeProvider;
 enum class EARFaceTrackingUpdate : uint8;
 
-class APPLEARKIT_API IAppleARKitFaceSupport :
+class IAppleARKitFaceSupport :
 	public IModularFeature
 {
 public:
diff --git a/Engine/Plugins/Runtime/AR/AppleAR/AppleARKit/Source/AppleARKit/Public/AppleARKitPoseTrackingLiveLink.h b/Engine/Plugins/Runtime/AR/AppleAR/AppleARKit/Source/AppleARKit/Public/AppleARKitPoseTrackingLiveLink.h
--- a/Engine/Plugins/Runtime/AR/AppleAR/AppleARKit/Source/AppleARKit/Public/AppleARKitPoseTrackingLiveLink.h
+++ b/Engine/Plugins/Runtime/AR/AppleAR/AppleARKit/Source/AppleARKit/Public/AppleARKitPoseTrackingLiveLink.h
@@ -13,7 +13,7 @@
 
 struct FAppleARKitAnchorData;
 
-class APPLEARKIT_API IAppleARKitPoseTrackingLiveLink
+class IAppleARKitPoseTrackingLiveLink
 	: public IModularFeature
 {
 public:
